#UE4 Add many missing copyrights

diff --git a/Engine/Source/Runtime/ClothingSystemRuntimeCommon/Private/ClothingAsset.cpp b/Engine/Source/Runtime/ClothingSystemRuntimeCommon/Private/ClothingAsset.cpp
--- a/Engine/Source/Runtime/ClothingSystemRuntimeCommon/Private/ClothingAsset.cpp
+++ b/Engine/Source/Runtime/ClothingSystemRuntimeCommon/Private/ClothingAsset.cpp
@@ -1,3 +1,5 @@
+// Copyright 1998-2019 Epic Games, Inc. All Rights Reserved.
+
 #include "ClothingAsset.h"
 
 #include "Utils/ClothingMeshUtils.h"
diff --git a/Engine/Source/Runtime/ClothingSystemRuntimeInterface/Private/ClothConfig.cpp b/Engine/Source/Runtime/ClothingSystemRuntimeInterface/Private/ClothConfig.cpp
--- a/Engine/Source/Runtime/ClothingSystemRuntimeInterface/Private/ClothConfig.cpp
+++ b/Engine/Source/Runtime/ClothingSystemRuntimeInterface/Private/ClothConfig.cpp
@@ -1,3 +1,5 @@
+// Copyright 1998-2019 Epic Games, Inc. All Rights Reserved.
+
 #include "ClothConfig.h"
 
 UClothConfigBase::UClothConfigBase()
diff --git a/Engine/Source/Runtime/ClothingSystemRuntimeInterface/Private/ClothPhysicalMeshData.cpp b/Engine/Source/Runtime/ClothingSystemRuntimeInterface/Private/ClothPhysicalMeshData.cpp
--- a/Engine/Source/Runtime/ClothingSystemRuntimeInterface/Private/ClothPhysicalMeshData.cpp
+++ b/Engine/Source/Runtime/ClothingSystemRuntimeInterface/Private/ClothPhysicalMeshData.cpp
@@ -1,3 +1,5 @@
+// Copyright 1998-2019 Epic Games, Inc. All Rights Reserved.
+
 #include "ClothPhysicalMeshData.h"
 
 UClothPhysicalMeshDataBase::UClothPhysicalMeshDataBase()
diff --git a/Engine/Source/Runtime/ClothingSystemRuntimeInterface/Public/ClothConfig.h b/Engine/Source/Runtime/ClothingSystemRuntimeInterface/Public/ClothConfig.h
--- a/Engine/Source/Runtime/ClothingSystemRuntimeInterface/Public/ClothConfig.h
+++ b/Engine/Source/Runtime/ClothingSystemRuntimeInterface/Public/ClothConfig.h
@@ -1,3 +1,5 @@
+// Copyright 1998-2019 Epic Games, Inc. All Rights Reserved.
+
 #pragma once
 
 #include "UObject/ObjectMacros.h"
diff --git a/Engine/Source/Runtime/ClothingSystemRuntimeInterface/Public/ClothVertBoneData.h b/Engine/Source/Runtime/ClothingSystemRuntimeInterface/Public/ClothVertBoneData.h
--- a/Engine/Source/Runtime/ClothingSystemRuntimeInterface/Public/ClothVertBoneData.h
+++ b/Engine/Source/Runtime/ClothingSystemRuntimeInterface/Public/ClothVertBoneData.h
@@ -1,3 +1,5 @@
+// Copyright 1998-2019 Epic Games, Inc. All Rights Reserved.
+
 #pragma once
 
 #include "ClothVertBoneData.generated.h"
diff --git a/Engine/Source/Runtime/ClothingSystemRuntimeNv/Private/ClothConfigNv.cpp b/Engine/Source/Runtime/ClothingSystemRuntimeNv/Private/ClothConfigNv.cpp
--- a/Engine/Source/Runtime/ClothingSystemRuntimeNv/Private/ClothConfigNv.cpp
+++ b/Engine/Source/Runtime/ClothingSystemRuntimeNv/Private/ClothConfigNv.cpp
@@ -1,3 +1,5 @@
+// Copyright 1998-2019 Epic Games, Inc. All Rights Reserved.
+
 #include "ClothConfigNv.h"
 
 UClothConfigNv::UClothConfigNv()
diff --git a/Engine/Source/Runtime/Experimental/Chaos/Private/Chaos/ParticleIterator.cpp b/Engine/Source/Runtime/Experimental/Chaos/Private/Chaos/ParticleIterator.cpp
--- a/Engine/Source/Runtime/Experimental/Chaos/Private/Chaos/ParticleIterator.cpp
+++ b/Engine/Source/Runtime/Experimental/Chaos/Private/Chaos/ParticleIterator.cpp
@@ -1,3 +1,5 @@
+// Copyright 1998-2019 Epic Games, Inc. All Rights Reserved.
+
 #include "Chaos/ParticleIterator.h"
 
 namespace Chaos
