Rename ClothConfig files to ClothConfigBase in line with the class name definition.

diff --git a/Engine/Plugins/Experimental/ChaosCloth/Source/Chaos/Public/ChaosCloth/ChaosClothConfig.h b/Engine/Plugins/Experimental/ChaosCloth/Source/Chaos/Public/ChaosCloth/ChaosClothConfig.h
--- a/Engine/Plugins/Experimental/ChaosCloth/Source/Chaos/Public/ChaosCloth/ChaosClothConfig.h
+++ b/Engine/Plugins/Experimental/ChaosCloth/Source/Chaos/Public/ChaosCloth/ChaosClothConfig.h
@@ -2,7 +2,7 @@
 #pragma once
 #include "Components/SkeletalMeshComponent.h"
 #include "Engine/SkeletalMesh.h"
-#include "ClothConfig.h"
+#include "ClothConfigBase.h"
 #include "CoreMinimal.h"
 #include "ChaosClothConfig.generated.h"
 
diff --git a/Engine/Source/Runtime/ClothingSystemRuntimeCommon/Public/ClothingAsset.h b/Engine/Source/Runtime/ClothingSystemRuntimeCommon/Public/ClothingAsset.h
--- a/Engine/Source/Runtime/ClothingSystemRuntimeCommon/Public/ClothingAsset.h
+++ b/Engine/Source/Runtime/ClothingSystemRuntimeCommon/Public/ClothingAsset.h
@@ -2,7 +2,7 @@
 #pragma once
 
 #include "ClothingAssetBase.h"
-#include "ClothConfig.h"
+#include "ClothConfigBase.h"
 #include "ClothLODData.h"
 #include "ClothingSimulationInteractor.h"
 
diff --git a/Engine/Source/Runtime/ClothingSystemRuntimeInterface/Private/ClothConfig.cpp b/Engine/Source/Runtime/ClothingSystemRuntimeInterface/Private/ClothConfig.cpp
deleted file mode 100644
diff --git a/Engine/Source/Runtime/ClothingSystemRuntimeInterface/Private/ClothConfigBase.cpp b/Engine/Source/Runtime/ClothingSystemRuntimeInterface/Private/ClothConfigBase.cpp
new file mode 100644
--- /dev/null
+++ b/Engine/Source/Runtime/ClothingSystemRuntimeInterface/Private/ClothConfigBase.cpp
@@ -0,0 +1,9 @@
+// Copyright 1998-2019 Epic Games, Inc. All Rights Reserved.
+
+#include "ClothConfigBase.h"
+
+UClothConfigBase::UClothConfigBase()
+{}
+
+UClothConfigBase::~UClothConfigBase()
+{}
diff --git a/Engine/Source/Runtime/ClothingSystemRuntimeInterface/Public/ClothConfig.h b/Engine/Source/Runtime/ClothingSystemRuntimeInterface/Public/ClothConfig.h
deleted file mode 100644
diff --git a/Engine/Source/Runtime/ClothingSystemRuntimeInterface/Public/ClothConfigBase.h b/Engine/Source/Runtime/ClothingSystemRuntimeInterface/Public/ClothConfigBase.h
new file mode 100644
--- /dev/null
+++ b/Engine/Source/Runtime/ClothingSystemRuntimeInterface/Public/ClothConfigBase.h
@@ -0,0 +1,40 @@
+// Copyright 1998-2019 Epic Games, Inc. All Rights Reserved.
+
+#pragma once
+
+#include "UObject/ObjectMacros.h"
+#include "ClothConfigBase.generated.h"
+
+class UClothingAssetBase;
+
+/**
+ * Base class for simulator specific simulation controls.
+ * Each cloth instance on a skeletal mesh can have a unique cloth config
+ */
+UCLASS(Abstract, DefaultToInstanced)
+class CLOTHINGSYSTEMRUNTIMEINTERFACE_API UClothConfigBase : public UObject
+{
+	GENERATED_BODY()
+public:
+	UClothConfigBase();
+	virtual ~UClothConfigBase();
+
+	virtual bool HasSelfCollision() const
+	{
+		unimplemented(); return false;
+	}
+};
+
+
+/*
+These settings are shared between all instances on a skeletal mesh
+*/
+UCLASS(Abstract)
+class CLOTHINGSYSTEMRUNTIMEINTERFACE_API UClothSharedSimConfigBase : public UObject
+{
+	GENERATED_BODY()
+public:
+	UClothSharedSimConfigBase() {};
+	virtual ~UClothSharedSimConfigBase() {};
+	
+};
diff --git a/Engine/Source/Runtime/ClothingSystemRuntimeNv/Public/ClothConfigNv.h b/Engine/Source/Runtime/ClothingSystemRuntimeNv/Public/ClothConfigNv.h
--- a/Engine/Source/Runtime/ClothingSystemRuntimeNv/Public/ClothConfigNv.h
+++ b/Engine/Source/Runtime/ClothingSystemRuntimeNv/Public/ClothConfigNv.h
@@ -2,7 +2,7 @@
 #pragma once
 
 #include "Engine/SkeletalMesh.h"
-#include "ClothConfig.h"
+#include "ClothConfigBase.h"
 #include "ClothConfigNv.generated.h"
 
 namespace nvidia
