Rename ClothPhysicalMeshData files to ClothPhysicalMeshDataBase to match the other class/file names and allow for a common runtime version.

diff --git a/Engine/Source/Editor/ClothPainter/Private/SClothAssetSelector.cpp b/Engine/Source/Editor/ClothPainter/Private/SClothAssetSelector.cpp
--- a/Engine/Source/Editor/ClothPainter/Private/SClothAssetSelector.cpp
+++ b/Engine/Source/Editor/ClothPainter/Private/SClothAssetSelector.cpp
@@ -8,7 +8,7 @@
 #include "Widgets/Input/SButton.h"
 #include "Widgets/Images/SImage.h"
 #include "ClothingAsset.h"
-#include "ClothPhysicalMeshData.h"
+#include "ClothPhysicalMeshDataBase.h"
 #include "Framework/MultiBox/MultiBoxBuilder.h"
 #include "Widgets/Input/SCheckBox.h"
 #include "Widgets/Text/SInlineEditableTextBlock.h"
diff --git a/Engine/Source/Runtime/ClothingSystemRuntimeCommon/Private/Utils/ClothingMeshUtils.cpp b/Engine/Source/Runtime/ClothingSystemRuntimeCommon/Private/Utils/ClothingMeshUtils.cpp
--- a/Engine/Source/Runtime/ClothingSystemRuntimeCommon/Private/Utils/ClothingMeshUtils.cpp
+++ b/Engine/Source/Runtime/ClothingSystemRuntimeCommon/Private/Utils/ClothingMeshUtils.cpp
@@ -2,7 +2,7 @@
 
 #include "Utils/ClothingMeshUtils.h"
 
-#include "ClothPhysicalMeshData.h"
+#include "ClothPhysicalMeshDataBase.h"
 
 #include "Math/UnrealMathUtility.h"
 #include "Logging/LogMacros.h"
diff --git a/Engine/Source/Runtime/ClothingSystemRuntimeCommon/Public/ClothLODData.h b/Engine/Source/Runtime/ClothingSystemRuntimeCommon/Public/ClothLODData.h
--- a/Engine/Source/Runtime/ClothingSystemRuntimeCommon/Public/ClothLODData.h
+++ b/Engine/Source/Runtime/ClothingSystemRuntimeCommon/Public/ClothLODData.h
@@ -1,7 +1,7 @@
 // Copyright 1998-2019 Epic Games, Inc. All Rights Reserved.
 #pragma once
 
-#include "ClothPhysicalMeshData.h"
+#include "ClothPhysicalMeshDataBase.h"
 
 #include "ClothCollisionData.h"
 #include "PointWeightMap.h"
diff --git a/Engine/Source/Runtime/ClothingSystemRuntimeInterface/Private/ClothPhysicalMeshData.cpp b/Engine/Source/Runtime/ClothingSystemRuntimeInterface/Private/ClothPhysicalMeshData.cpp
deleted file mode 100644
diff --git a/Engine/Source/Runtime/ClothingSystemRuntimeInterface/Private/ClothPhysicalMeshDataBase.cpp b/Engine/Source/Runtime/ClothingSystemRuntimeInterface/Private/ClothPhysicalMeshDataBase.cpp
new file mode 100644
--- /dev/null
+++ b/Engine/Source/Runtime/ClothingSystemRuntimeInterface/Private/ClothPhysicalMeshDataBase.cpp
@@ -0,0 +1,64 @@
+// Copyright 1998-2019 Epic Games, Inc. All Rights Reserved.
+
+#include "ClothPhysicalMeshDataBase.h"
+
+UClothPhysicalMeshDataBase::UClothPhysicalMeshDataBase()
+	: NumFixedVerts(0)
+	, MaxBoneWeights(0)
+{}
+
+UClothPhysicalMeshDataBase::~UClothPhysicalMeshDataBase()
+{}
+
+void UClothPhysicalMeshDataBase::Reset(const int32 InNumVerts)
+{
+	Vertices.Reset();
+	Normals.Reset();
+#if WITH_EDITORONLY_DATA
+	VertexColors.Reset();
+#endif // #if WITH_EDITORONLY_DATA
+	InverseMasses.Reset();
+	BoneData.Reset();
+
+	Vertices.AddDefaulted(InNumVerts);
+	Normals.AddDefaulted(InNumVerts);
+#if WITH_EDITORONLY_DATA
+	VertexColors.AddDefaulted(InNumVerts);
+#endif //#if WITH_EDITORONLY_DATA
+
+	InverseMasses.AddDefaulted(InNumVerts);
+	BoneData.AddDefaulted(InNumVerts);
+
+	NumFixedVerts = 0;
+	MaxBoneWeights = 0;
+}
+
+TArray<float>* UClothPhysicalMeshDataBase::GetFloatArray(const uint32 Id) const
+{
+	check(IdToArray.Contains(Id));
+	return IdToArray[Id];
+}
+
+TArray<uint32> UClothPhysicalMeshDataBase::GetFloatArrayIds() const
+{
+	TArray<uint32> Keys; 
+	IdToArray.GetKeys(Keys);
+	return Keys;
+}
+
+TArray<TArray<float>*> UClothPhysicalMeshDataBase::GetFloatArrays() const
+{
+	TArray<TArray<float>*> Values;
+	IdToArray.GenerateValueArray(Values);
+	return Values;
+}
+
+void UClothPhysicalMeshDataBase::RegisterFloatArray(
+	const uint32 Id,
+	TArray<float> *Array)
+{
+	check(Id != INDEX_NONE);
+	check(Array != nullptr);
+	check(!IdToArray.Contains(Id) || IdToArray[Id] == Array);
+	IdToArray.Add(Id, Array);
+}
diff --git a/Engine/Source/Runtime/ClothingSystemRuntimeInterface/Public/ClothPhysicalMeshData.h b/Engine/Source/Runtime/ClothingSystemRuntimeInterface/Public/ClothPhysicalMeshData.h
deleted file mode 100644
diff --git a/Engine/Source/Runtime/ClothingSystemRuntimeInterface/Public/ClothPhysicalMeshDataBase.h b/Engine/Source/Runtime/ClothingSystemRuntimeInterface/Public/ClothPhysicalMeshDataBase.h
new file mode 100644
--- /dev/null
+++ b/Engine/Source/Runtime/ClothingSystemRuntimeInterface/Public/ClothPhysicalMeshDataBase.h
@@ -0,0 +1,114 @@
+// Copyright 1998-2019 Epic Games, Inc. All Rights Reserved.
+#pragma once
+
+#include "UObject/Object.h"
+
+#include "ClothVertBoneData.h"
+#include "Containers/Array.h"
+#include "Containers/Map.h"
+
+#include "ClothPhysicalMeshDataBase.generated.h"
+
+class UClothConfigBase;
+
+/**
+ * Simulation mesh points, topology, and spatial parameters defined on that 
+ * topology.
+ *
+ * Created curing asset import or created from a skeletal mesh.
+ */
+UCLASS()
+class CLOTHINGSYSTEMRUNTIMEINTERFACE_API UClothPhysicalMeshDataBase : public UObject
+{
+	GENERATED_BODY()
+public:
+	UClothPhysicalMeshDataBase();
+	virtual ~UClothPhysicalMeshDataBase();
+
+	virtual void 
+	Reset(const int32 InNumVerts);
+
+	// TODO: Ryan - Rename to ClearWeightedParameters()?
+	/** Clear out any target properties in this physical mesh. */
+	virtual void 
+	ClearParticleParameters()
+	{}
+
+	/** Callback invoked from \c UClothingAssetBase::BuildSelfCollisionData(), 
+	 * implement in derived classes. 
+	 */
+	virtual void 
+	BuildSelfCollisionData(const UClothConfigBase* ClothConfig)
+	{ unimplemented(); }
+
+	/** Retrieve a registered vertex weight array by unique @param Id. */
+	TArray<float>* 
+	GetFloatArray(const uint32 Id) const;
+
+	/** Get ids for all registered weight arrays. */
+	TArray<uint32>
+	GetFloatArrayIds() const;
+
+	/** Get all registered weight arrays. */
+	TArray<TArray<float>*>
+	GetFloatArrays() const;
+
+	/** Returns an \c Enum mapping float array Id's to names used in the UI. */
+	virtual UEnum* 
+	GetFloatArrayTargets() const
+	{ unimplemented(); return nullptr; }
+
+
+	virtual bool IsFullyKinematic() const
+	{ unimplemented(); return true; }
+	virtual bool IsFixed(const uint16 X, const float Threshold=0.1f) const
+	{ unimplemented(); return false; }
+	virtual bool IsFixed(const uint16 X, const uint16 Y, const uint16 Z, const float Threshold=0.1f) const
+	{ unimplemented(); return false; }
+
+protected:
+	/** Register an @param Array keyed by a unique @param Id. */
+	void RegisterFloatArray(const uint32 Id, TArray<float> *Array);
+
+public:
+	// Positions of each simulation vertex
+	UPROPERTY(EditAnywhere, Category = SimMesh)
+	TArray<FVector> Vertices;
+
+	// Normal at each vertex
+	UPROPERTY(EditAnywhere, Category = SimMesh)
+	TArray<FVector> Normals;
+
+#if WITH_EDITORONLY_DATA
+	// Color at each vertex
+	UPROPERTY(EditAnywhere, Category = SimMesh)
+	TArray<FColor> VertexColors;
+#endif // WITH_EDITORONLY_DATA
+
+	// Indices of the simulation mesh triangles
+	UPROPERTY(EditAnywhere, Category = SimMesh)
+	TArray<uint32> Indices;
+
+	// Inverse mass for each vertex in the physical mesh
+	UPROPERTY(EditAnywhere, Category = SimMesh)
+	TArray<float> InverseMasses;
+
+	// Indices and weights for each vertex, used to skin the mesh to create the reference pose
+	UPROPERTY(EditAnywhere, Category = SimMesh)
+	TArray<FClothVertBoneData> BoneData;
+
+	// Number of fixed verts in the simulation mesh (fixed verts are just skinned and do not simulate)
+	UPROPERTY(EditAnywhere, Category = SimMesh)
+	int32 NumFixedVerts;
+
+	// Maximum number of bone weights of any vetex
+	UPROPERTY(EditAnywhere, Category = SimMesh)
+	int32 MaxBoneWeights;
+
+	// Valid indices to use for self collisions (reduced set of Indices)
+	UPROPERTY(EditAnywhere, Category = SimMesh)
+	TArray<uint32> SelfCollisionIndices;
+
+private:
+	TMap<uint32, TArray<float>*> IdToArray;
+};
diff --git a/Engine/Source/Runtime/ClothingSystemRuntimeNv/Public/ClothPhysicalMeshDataNv.h b/Engine/Source/Runtime/ClothingSystemRuntimeNv/Public/ClothPhysicalMeshDataNv.h
--- a/Engine/Source/Runtime/ClothingSystemRuntimeNv/Public/ClothPhysicalMeshDataNv.h
+++ b/Engine/Source/Runtime/ClothingSystemRuntimeNv/Public/ClothPhysicalMeshDataNv.h
@@ -1,7 +1,7 @@
 // Copyright 1998-2019 Epic Games, Inc. All Rights Reserved.
 #pragma once
 
-#include "ClothPhysicalMeshData.h"
+#include "ClothPhysicalMeshDataBase.h"
 #include "ClothPhysicalMeshDataNv.generated.h"
 
 /** The possible targets for a physical mesh point weight map. */
