Set the correct default clothing simulation factory based on the target build defines.
Add a new function that returns the correct default clothing simulation factory.
Add a new p.Cloth.DefaultClothingFactoryClass console variable to change the default at runtime.
Update the SkeletalMeshComponent code to use the correct default clothing simulation factory.
Deprecate the IClothingSimulationFactoryClassProvider::GetDefaultSimulationFactoryClass() and add a new IClothingSimulationFactoryClassProvider::GetSimulationFactoryClass() function to remove any ambiguity about what the correct way to get the provider's class, or the default class.

diff --git a/Engine/Plugins/Experimental/ChaosCloth/Source/Chaos/Private/ChaosCloth/ChaosClothModule.cpp b/Engine/Plugins/Experimental/ChaosCloth/Source/Chaos/Private/ChaosCloth/ChaosClothModule.cpp
--- a/Engine/Plugins/Experimental/ChaosCloth/Source/Chaos/Private/ChaosCloth/ChaosClothModule.cpp
+++ b/Engine/Plugins/Experimental/ChaosCloth/Source/Chaos/Private/ChaosCloth/ChaosClothModule.cpp
@@ -27,7 +27,7 @@ class FChaosClothModule : public IChaosClothModuleInterface, public IClothingSim
 #endif
     }
 
-	UClass* GetDefaultSimulationFactoryClass()
+	TSubclassOf<UClothingSimulationFactory> GetClothingSimulationFactoryClass() const override
 	{
 #if WITH_CHAOS
 		return UChaosClothingSimulationFactory::StaticClass();
diff --git a/Engine/Source/Editor/Persona/Private/PreviewSceneCustomizations.cpp b/Engine/Source/Editor/Persona/Private/PreviewSceneCustomizations.cpp
--- a/Engine/Source/Editor/Persona/Private/PreviewSceneCustomizations.cpp
+++ b/Engine/Source/Editor/Persona/Private/PreviewSceneCustomizations.cpp
@@ -306,7 +306,7 @@ void FPreviewSceneDescriptionCustomization::CustomizeDetails(IDetailLayoutBuilde
 		for (const auto& ClassProvider : ClassProviders)
 		{
 			// Populate cloth factory list
-			ClothSimulationFactoryList.Add(MakeShared<TSubclassOf<class UClothingSimulationFactory>>(ClassProvider->GetDefaultSimulationFactoryClass()));
+			ClothSimulationFactoryList.Add(MakeShared<TSubclassOf<class UClothingSimulationFactory>>(ClassProvider->GetClothingSimulationFactoryClass()));
 		}
 
 		DetailBuilder.EditCategory("Physics")
diff --git a/Engine/Source/Runtime/ClothingSystemRuntimeInterface/Private/ClothingSimulationFactory.cpp b/Engine/Source/Runtime/ClothingSystemRuntimeInterface/Private/ClothingSimulationFactory.cpp
--- a/Engine/Source/Runtime/ClothingSystemRuntimeInterface/Private/ClothingSimulationFactory.cpp
+++ b/Engine/Source/Runtime/ClothingSystemRuntimeInterface/Private/ClothingSimulationFactory.cpp
@@ -1,5 +1,49 @@
 // Copyright 1998-2019 Epic Games, Inc. All Rights Reserved.
 
 #include "ClothingSimulationFactory.h"
+#include "HAL/IConsoleManager.h"
+#include "Features/IModularFeatures.h"
 
-const FName IClothingSimulationFactoryClassProvider::FeatureName = TEXT("ClothingSimulationFactoryClassProvider");
\ No newline at end of file
+const FName IClothingSimulationFactoryClassProvider::FeatureName = TEXT("ClothingSimulationFactoryClassProvider");
+
+namespace ClothingSimulationFactoryConsoleVariables
+{
+	TAutoConsoleVariable<FString> CVarDefaultClothingSimulationFactoryClass(
+		TEXT("p.Cloth.DefaultClothingSimulationFactoryClass"),
+#if WITH_CHAOS_CLOTHING
+		TEXT("ChaosClothingSimulationFactory"),  // Chaos is the default provider when Chaos Cloth is enabled
+#elif WITH_APEX_CLOTHING
+		TEXT("ClothingSimulationFactoryNv"),  // otherwise it's nv cloth
+#endif
+		TEXT("The class name of the default clothing simulation factory.\n")
+		TEXT("Known providers are:\n")
+#if WITH_CHAOS_CLOTHING
+		TEXT("ChaosClothingSimulationFactory\n")
+#endif
+#if WITH_APEX_CLOTHING
+		TEXT("ClothingSimulationFactoryNv\n")
+#endif
+		, ECVF_Cheat);
+}
+
+TSubclassOf<class UClothingSimulationFactory> UClothingSimulationFactory::GetDefaultClothingSimulationFactoryClass()
+{
+	TSubclassOf<UClothingSimulationFactory> DefaultClothingSimulationFactoryClass = nullptr;
+
+	const FString DefaultClothingSimulationFactoryClassName = ClothingSimulationFactoryConsoleVariables::CVarDefaultClothingSimulationFactoryClass.GetValueOnAnyThread();
+
+	TArray<IClothingSimulationFactoryClassProvider*> ClassProviders = IModularFeatures::Get().GetModularFeatureImplementations<IClothingSimulationFactoryClassProvider>(IClothingSimulationFactoryClassProvider::FeatureName);
+	for (const auto& ClassProvider : ClassProviders)
+	{
+		if (ClassProvider)
+		{
+			DefaultClothingSimulationFactoryClass = ClassProvider->GetClothingSimulationFactoryClass();
+
+			if (DefaultClothingSimulationFactoryClass->GetName() == DefaultClothingSimulationFactoryClassName)
+			{
+				break;
+			}
+		}
+	}
+	return DefaultClothingSimulationFactoryClass;
+}
diff --git a/Engine/Source/Runtime/ClothingSystemRuntimeInterface/Public/ClothingSimulationFactory.h b/Engine/Source/Runtime/ClothingSystemRuntimeInterface/Public/ClothingSimulationFactory.h
--- a/Engine/Source/Runtime/ClothingSystemRuntimeInterface/Public/ClothingSimulationFactory.h
+++ b/Engine/Source/Runtime/ClothingSystemRuntimeInterface/Public/ClothingSimulationFactory.h
@@ -6,6 +6,7 @@
 
 #include "UObject/Object.h"
 #include "Features/IModularFeature.h"
+#include "Templates/SubclassOf.h"
 
 #include "ClothingSimulationFactory.generated.h"
 
@@ -25,7 +26,13 @@ public:
 	// Called by the engine to get the default clothing simulation factory to use
 	// for skeletal mesh components (see USkeletalMeshComponent constructor).
 	// Returns Factory class for simulations or nullptr to disable clothing simulation
-	virtual UClass* GetDefaultSimulationFactoryClass() = 0;
+	UE_DEPRECATED(4.25, "GetDefaultSimulationFactoryClass() has been deprecated. Use IClothingSimulationFactoryClassProvider::GetSimulationFactoryClass() or UClothingSimulationFactory::GetDefaultClothingSimulationFactoryClass() instead.")
+	virtual UClass* GetDefaultSimulationFactoryClass() { return nullptr; }
+
+	// Called by the engine to get the clothing simulation factory associated with this
+	// provider for skeletal mesh components (see USkeletalMeshComponent constructor).
+	// Returns Factory class for simulations or nullptr to disable clothing simulation
+	virtual TSubclassOf<class UClothingSimulationFactory> GetClothingSimulationFactoryClass() const = 0;
 };
 
 // Any clothing simulation factory should derive from this interface object to interact with the engine
@@ -35,6 +42,10 @@ class CLOTHINGSYSTEMRUNTIMEINTERFACE_API UClothingSimulationFactory : public UOb
 	GENERATED_BODY()
 
 public:
+	// Return the default clothing simulation factory class as set by the build or by
+	// the p.Cloth.DefaultClothingSimulationFactoryClass console variable if any available.
+	// Otherwise return the last registered factory.
+	static TSubclassOf<class UClothingSimulationFactory> GetDefaultClothingSimulationFactoryClass();
 
 	// Create a simulation object for a skeletal mesh to use (see IClothingSimulation)
 	virtual IClothingSimulation* CreateSimulation()
diff --git a/Engine/Source/Runtime/ClothingSystemRuntimeNv/Private/Assets/ClothingAssetNv.cpp b/Engine/Source/Runtime/ClothingSystemRuntimeNv/Private/Assets/ClothingAssetNv.cpp
--- a/Engine/Source/Runtime/ClothingSystemRuntimeNv/Private/Assets/ClothingAssetNv.cpp
+++ b/Engine/Source/Runtime/ClothingSystemRuntimeNv/Private/Assets/ClothingAssetNv.cpp
@@ -73,7 +73,7 @@ UClothingAssetNv::UClothingAssetNv(const FObjectInitializer& ObjectInitializer)
 	{
 		IClothingSimulationFactoryClassProvider* const Provider = ClassProviders[Index];
 		check(Provider);
-		ClothingSimulationFactory = Provider->GetDefaultSimulationFactoryClass();
+		ClothingSimulationFactory = Provider->GetClothingSimulationFactoryClass();
 
 		UClass* SimFactoryClass = *ClothingSimulationFactory;
 		if (SimFactoryClass)
diff --git a/Engine/Source/Runtime/ClothingSystemRuntimeNv/Private/ClothingSystemRuntimeModuleNv.cpp b/Engine/Source/Runtime/ClothingSystemRuntimeNv/Private/ClothingSystemRuntimeModuleNv.cpp
--- a/Engine/Source/Runtime/ClothingSystemRuntimeNv/Private/ClothingSystemRuntimeModuleNv.cpp
+++ b/Engine/Source/Runtime/ClothingSystemRuntimeNv/Private/ClothingSystemRuntimeModuleNv.cpp
@@ -73,7 +73,7 @@ nv::cloth::ClothMeshQuadifier* FClothingSystemRuntimeModuleNv::GetMeshQuadifier(
 	return nullptr;
 }
 
-UClass* FClothingSystemRuntimeModuleNv::GetDefaultSimulationFactoryClass()
+TSubclassOf<UClothingSimulationFactory> FClothingSystemRuntimeModuleNv::GetClothingSimulationFactoryClass() const
 {
 #if WITH_NVCLOTH
 	return UClothingSimulationFactoryNv::StaticClass();
diff --git a/Engine/Source/Runtime/ClothingSystemRuntimeNv/Public/ClothingSystemRuntimeModuleNv.h b/Engine/Source/Runtime/ClothingSystemRuntimeNv/Public/ClothingSystemRuntimeModuleNv.h
--- a/Engine/Source/Runtime/ClothingSystemRuntimeNv/Public/ClothingSystemRuntimeModuleNv.h
+++ b/Engine/Source/Runtime/ClothingSystemRuntimeNv/Public/ClothingSystemRuntimeModuleNv.h
@@ -28,7 +28,7 @@ public:
 	nv::cloth::ClothMeshQuadifier* GetMeshQuadifier();
 
 	// IClothingSimulationFactoryClassProvider Interface
-	virtual UClass* GetDefaultSimulationFactoryClass() override;
+	virtual TSubclassOf<UClothingSimulationFactory> GetClothingSimulationFactoryClass() const override;
 	//////////////////////////////////////////////////////////////////////////
 
 private:
diff --git a/Engine/Source/Runtime/Engine/Private/Components/SkeletalMeshComponent.cpp b/Engine/Source/Runtime/Engine/Private/Components/SkeletalMeshComponent.cpp
--- a/Engine/Source/Runtime/Engine/Private/Components/SkeletalMeshComponent.cpp
+++ b/Engine/Source/Runtime/Engine/Private/Components/SkeletalMeshComponent.cpp
@@ -233,16 +233,7 @@ USkeletalMeshComponent::USkeletalMeshComponent(const FObjectInitializer& ObjectI
 	CachedAnimCurveUidVersion = 0;
 	ResetRootBodyIndex();
 
-	TArray<IClothingSimulationFactoryClassProvider*> ClassProviders = IModularFeatures::Get().GetModularFeatureImplementations<IClothingSimulationFactoryClassProvider>(IClothingSimulationFactoryClassProvider::FeatureName);
-
-	ClothingSimulationFactory = nullptr;
-	for (int32 Index = ClassProviders.Num() - 1; Index >= 0 && !*ClothingSimulationFactory; --Index)  // We use the last provider in the list so plugins/modules can override ours
-	{
-		IClothingSimulationFactoryClassProvider* const Provider = ClassProviders[Index];
-		check(Provider);
-
-		ClothingSimulationFactory = Provider->GetDefaultSimulationFactoryClass();
-	}
+	ClothingSimulationFactory = UClothingSimulationFactory::GetDefaultClothingSimulationFactoryClass();
 
 	ClothingSimulation = nullptr;
 	ClothingSimulationContext = nullptr;
@@ -531,18 +522,10 @@ void USkeletalMeshComponent::OnRegister()
 	}
 
 #if WITH_APEX_CLOTHING || WITH_CHAOS_CLOTHING
-	
 	// If we don't have a valid simulation factory - check to see if we have an available default to use instead
-	if(*ClothingSimulationFactory == nullptr)
+	if (*ClothingSimulationFactory == nullptr)
 	{
-		TArray<IClothingSimulationFactoryClassProvider*> ClassProviders = IModularFeatures::Get().GetModularFeatureImplementations<IClothingSimulationFactoryClassProvider>(IClothingSimulationFactoryClassProvider::FeatureName);
-		for (int32 Index = ClassProviders.Num() - 1; Index >= 0 && !*ClothingSimulationFactory; --Index)  // We use the last provider in the list so plugins/modules can override ours
-		{
-			IClothingSimulationFactoryClassProvider* const Provider = ClassProviders[Index];
-			check(Provider);
-
-			ClothingSimulationFactory = Provider->GetDefaultSimulationFactoryClass();
-		}
+		ClothingSimulationFactory = UClothingSimulationFactory::GetDefaultClothingSimulationFactoryClass();
 	}
 
 	RecreateClothingActors();
