Fix CIS error when all cloth simulations types are disable in build.

diff --git a/Engine/Source/Runtime/ClothingSystemRuntimeInterface/Private/ClothingSimulationFactory.cpp b/Engine/Source/Runtime/ClothingSystemRuntimeInterface/Private/ClothingSimulationFactory.cpp
--- a/Engine/Source/Runtime/ClothingSystemRuntimeInterface/Private/ClothingSimulationFactory.cpp
+++ b/Engine/Source/Runtime/ClothingSystemRuntimeInterface/Private/ClothingSimulationFactory.cpp
@@ -14,6 +14,8 @@ namespace ClothingSimulationFactoryConsoleVariables
 		TEXT("ChaosClothingSimulationFactory"),  // Chaos is the default provider when Chaos Cloth is enabled
 #elif WITH_APEX_CLOTHING
 		TEXT("ClothingSimulationFactoryNv"),  // otherwise it's nv cloth
+#else
+		TEXT(""),  // otherwise it's none
 #endif
 		TEXT("The class name of the default clothing simulation factory.\n")
 		TEXT("Known providers are:\n")
