Fix crash while getting default cloth simulation factory when built with Chaos and running without the Chaos Cloth plugin.

diff --git a/Engine/Source/Runtime/ClothingSystemRuntimeInterface/Private/ClothingSimulationFactory.cpp b/Engine/Source/Runtime/ClothingSystemRuntimeInterface/Private/ClothingSimulationFactory.cpp
--- a/Engine/Source/Runtime/ClothingSystemRuntimeInterface/Private/ClothingSimulationFactory.cpp
+++ b/Engine/Source/Runtime/ClothingSystemRuntimeInterface/Private/ClothingSimulationFactory.cpp
@@ -34,16 +34,20 @@ TSubclassOf<class UClothingSimulationFactory> UClothingSimulationFactory::GetDef
 
 	const FString DefaultClothingSimulationFactoryClassName = ClothingSimulationFactoryConsoleVariables::CVarDefaultClothingSimulationFactoryClass.GetValueOnAnyThread();
 
-	TArray<IClothingSimulationFactoryClassProvider*> ClassProviders = IModularFeatures::Get().GetModularFeatureImplementations<IClothingSimulationFactoryClassProvider>(IClothingSimulationFactoryClassProvider::FeatureName);
+	const TArray<IClothingSimulationFactoryClassProvider*> ClassProviders = IModularFeatures::Get().GetModularFeatureImplementations<IClothingSimulationFactoryClassProvider>(IClothingSimulationFactoryClassProvider::FeatureName);
 	for (const auto& ClassProvider : ClassProviders)
 	{
 		if (ClassProvider)
 		{
-			DefaultClothingSimulationFactoryClass = ClassProvider->GetClothingSimulationFactoryClass();
-
-			if (DefaultClothingSimulationFactoryClass->GetName() == DefaultClothingSimulationFactoryClassName)
+			const TSubclassOf<UClothingSimulationFactory> ClothingSimulationFactoryClass = ClassProvider->GetClothingSimulationFactoryClass();
+			if (*ClothingSimulationFactoryClass)
 			{
-				break;
+				DefaultClothingSimulationFactoryClass = ClothingSimulationFactoryClass;
+
+				if (ClothingSimulationFactoryClass->GetName() == DefaultClothingSimulationFactoryClassName)
+				{
+					break;
+				}
 			}
 		}
 	}
