Add comments to GetDefaultClothingSimulationFactoryClass, and make the dereference to UClass more explicit.

diff --git a/Engine/Source/Runtime/ClothingSystemRuntimeInterface/Private/ClothingSimulationFactory.cpp b/Engine/Source/Runtime/ClothingSystemRuntimeInterface/Private/ClothingSimulationFactory.cpp
--- a/Engine/Source/Runtime/ClothingSystemRuntimeInterface/Private/ClothingSimulationFactory.cpp
+++ b/Engine/Source/Runtime/ClothingSystemRuntimeInterface/Private/ClothingSimulationFactory.cpp
@@ -40,10 +40,12 @@ TSubclassOf<class UClothingSimulationFactory> UClothingSimulationFactory::GetDef
 		if (ClassProvider)
 		{
 			const TSubclassOf<UClothingSimulationFactory> ClothingSimulationFactoryClass = ClassProvider->GetClothingSimulationFactoryClass();
-			if (*ClothingSimulationFactoryClass)
+			if (ClothingSimulationFactoryClass.Get() != nullptr)
 			{
+				// Always set the default to the last non null factory class (in case the search for the cvar doesn't yield any results)
 				DefaultClothingSimulationFactoryClass = ClothingSimulationFactoryClass;
 
+				// Early exit if the cvar string matches
 				if (ClothingSimulationFactoryClass->GetName() == DefaultClothingSimulationFactoryClassName)
 				{
 					break;
