Chaos Cloth - Makes the Clothing Simulation Interactor a BP class to replace types of otherwise Nv specific variables.

diff --git a/Engine/Source/Runtime/ClothingSystemRuntimeInterface/Public/ClothingSimulationInteractor.h b/Engine/Source/Runtime/ClothingSystemRuntimeInterface/Public/ClothingSimulationInteractor.h
--- a/Engine/Source/Runtime/ClothingSystemRuntimeInterface/Public/ClothingSimulationInteractor.h
+++ b/Engine/Source/Runtime/ClothingSystemRuntimeInterface/Public/ClothingSimulationInteractor.h
@@ -14,7 +14,7 @@ class IClothingSimulationContext;
  * Only write to the simulation and context during the call to Sync, as that is
  * guaranteed to be a safe place to access this data.
  */
-UCLASS(Abstract)
+UCLASS(Abstract, BlueprintType)
 class CLOTHINGSYSTEMRUNTIMEINTERFACE_API UClothingSimulationInteractor : public UObject
 {
 	GENERATED_BODY()
