ClothEditor:
- replace the dummy 3D preview mesh with a SkeletalMeshComponent
- when a tool finishes operating on the rest-space mesh, set the SkeletalMesh's sim mesh from the result
- add a ClothEditor3DViewportClient to render the sim mesh using debug render
- remove dead code relating to background visualization copied from the UV editor

#rnx

diff --git a/Engine/Source/Runtime/ClothingSystemRuntimeInterface/Public/ClothingAssetBase.h b/Engine/Source/Runtime/ClothingSystemRuntimeInterface/Public/ClothingAssetBase.h
--- a/Engine/Source/Runtime/ClothingSystemRuntimeInterface/Public/ClothingAssetBase.h
+++ b/Engine/Source/Runtime/ClothingSystemRuntimeInterface/Public/ClothingAssetBase.h
@@ -125,6 +125,7 @@ protected:
 
 	/** The asset factory should have access, as it will assign the asset guid when building assets */
 	friend class UClothingAssetFactory;
+	friend class UClothEditorMode;
 
 	/** Guid to identify this asset. Will be embedded into chunks that are created using this asset */
 	UPROPERTY()
