Rename/move:
- rename ClothEditor to ChaosClothAssetEditor
- rename ChaosClothAsset/ChaosClothAssetEditor to ChaosClothAssetEditor/ChaosClothAssetTools

diff --git a/Engine/Source/Runtime/ClothingSystemRuntimeInterface/Public/ClothingAssetBase.h b/Engine/Source/Runtime/ClothingSystemRuntimeInterface/Public/ClothingAssetBase.h
--- a/Engine/Source/Runtime/ClothingSystemRuntimeInterface/Public/ClothingAssetBase.h
+++ b/Engine/Source/Runtime/ClothingSystemRuntimeInterface/Public/ClothingAssetBase.h
@@ -125,7 +125,9 @@ protected:
 
 	/** The asset factory should have access, as it will assign the asset guid when building assets */
 	friend class UClothingAssetFactory;
-	friend class UClothEditorMode;
+
+	/* TODO: Remove this temporary access (JIRA UE-159698) */
+	friend class UChaosClothAssetEditorMode;	
 
 	/** Guid to identify this asset. Will be embedded into chunks that are created using this asset */
 	UPROPERTY()
