[AutoRTFM] Allow GetClothingSimulationFactoryClass to work from a transaction.

The Nvidia cloth code isn't transactionally-aware, but it is safe to look up the cloth class.

diff --git a/Engine/Source/Runtime/ClothingSystemRuntimeInterface/Private/ClothingSimulationFactory.cpp b/Engine/Source/Runtime/ClothingSystemRuntimeInterface/Private/ClothingSimulationFactory.cpp
--- a/Engine/Source/Runtime/ClothingSystemRuntimeInterface/Private/ClothingSimulationFactory.cpp
+++ b/Engine/Source/Runtime/ClothingSystemRuntimeInterface/Private/ClothingSimulationFactory.cpp
@@ -1,6 +1,7 @@
 // Copyright Epic Games, Inc. All Rights Reserved.
 
 #include "ClothingSimulationFactory.h"
+#include "AutoRTFM/AutoRTFM.h"
 #include "HAL/IConsoleManager.h"
 #include "Features/IModularFeatures.h"
 
@@ -19,6 +20,7 @@ namespace ClothingSimulationFactoryConsoleVariables
 		, ECVF_Cheat);
 }
 
+UE_AUTORTFM_ALWAYS_OPEN
 TSubclassOf<class UClothingSimulationFactory> UClothingSimulationFactory::GetDefaultClothingSimulationFactoryClass()
 {
 	TSubclassOf<UClothingSimulationFactory> DefaultClothingSimulationFactoryClass = nullptr;
