Cloth - Deprecated ImportedFilePath which was an Apex import path and is no longer being used, since CL 20450744 when the PersonaMeshDetails Apex import code has been removed.

#rnx

diff --git a/Engine/Source/Runtime/ClothingSystemRuntimeInterface/Public/ClothingAssetBase.h b/Engine/Source/Runtime/ClothingSystemRuntimeInterface/Public/ClothingAssetBase.h
--- a/Engine/Source/Runtime/ClothingSystemRuntimeInterface/Public/ClothingAssetBase.h
+++ b/Engine/Source/Runtime/ClothingSystemRuntimeInterface/Public/ClothingAssetBase.h
@@ -101,8 +101,8 @@ public:
 	}
 
 #if WITH_EDITORONLY_DATA
-	// If this asset was imported from a file, this will be the original path
-	UPROPERTY(VisibleAnywhere, Category = Import)
+	UE_DEPRECATED(5.6, "This is an old Apex file path which is no longer used and will be removed.")
+	UPROPERTY()
 	FString ImportedFilePath;
 #endif
 
