Skeletal Mesh Cloth - Deprecate a few ClothingAssetBase virtual methods that should only be in the derived ClothingAssetCommon due to their dependency to the Skeletal Mesh Persona editor systems instead of the abstract Skeletal Mesh cloth implementation.

#rnx

diff --git a/Engine/Source/Runtime/ClothingSystemRuntimeCommon/Private/ClothingAsset.cpp b/Engine/Source/Runtime/ClothingSystemRuntimeCommon/Private/ClothingAsset.cpp
--- a/Engine/Source/Runtime/ClothingSystemRuntimeCommon/Private/ClothingAsset.cpp
+++ b/Engine/Source/Runtime/ClothingSystemRuntimeCommon/Private/ClothingAsset.cpp
@@ -1077,6 +1077,17 @@ void UClothingAssetCommon::PostLoad()
 	}
 
 	EClothingCachedDataFlagsCommon ClothingCachedDataFlags = EClothingCachedDataFlagsCommon::None;
+
+	if (USkeletalMesh* const SkeletalMesh = Cast<USkeletalMesh>(GetOuter()))
+	{
+		// Use the skeletal mesh version here to prevent a circular dependency between ClothingAssetCommon and Engine/SkeletalMesh
+		// and to unify all cached data invalidations, while the InvalidateCachedData method is also deprecated from ClothingAssetBase.
+		if (SkeletalMesh->GetLinkerCustomVersion(FSkeletalMeshCustomVersion::GUID) < FSkeletalMeshCustomVersion::CachedClothInverseMasses)
+		{
+			ClothingCachedDataFlags |= EClothingCachedDataFlagsCommon::InverseMasses;
+		}
+	}
+
 	if (AnimPhysCustomVersion < FAnimPhysObjectVersion::CacheClothMeshInfluences)
 	{
 		ClothingCachedDataFlags |= EClothingCachedDataFlagsCommon::NumInfluences;
diff --git a/Engine/Source/Runtime/ClothingSystemRuntimeInterface/Public/ClothingAssetBase.h b/Engine/Source/Runtime/ClothingSystemRuntimeInterface/Public/ClothingAssetBase.h
--- a/Engine/Source/Runtime/ClothingSystemRuntimeInterface/Public/ClothingAssetBase.h
+++ b/Engine/Source/Runtime/ClothingSystemRuntimeInterface/Public/ClothingAssetBase.h
@@ -52,6 +52,7 @@ public:
 	PURE_VIRTUAL(UClothingAssetBase::UpdateAllLODBiasMappings,);
 
 	/** Add a new LOD class instance. */
+	UE_DEPRECATED(5.7, "Use UClothingAssetCommone::AddNewLod() instead.")
 	virtual int32 AddNewLod()
 	PURE_VIRTUAL(UClothingAssetBase::AddNewLod(), return INDEX_NONE;);
 
@@ -60,6 +61,7 @@ public:
 	*	When we transition between LODs we skin the incoming mesh to the outgoing mesh
 	*	in exactly the same way the render mesh is skinned to create a smooth swap
 	*/
+	UE_DEPRECATED(5.7, "Use UClothingAssetCommon::BuildLodTransitionData() instead.")
 	virtual void BuildLodTransitionData()
 	PURE_VIRTUAL(UClothingAssetBase::BuildLodTransitionData(), );
 #endif
@@ -69,6 +71,7 @@ public:
 	 * Called on the clothing asset when the base data (physical mesh, config etc.)
 	 * has changed, so any intermediate generated data can be regenerated.
 	 */
+	UE_DEPRECATED(5.7, "Use UClothingAssetCommon::InvalidateAllCachedData() instead.")
 	virtual void InvalidateAllCachedData()
 	PURE_VIRTUAL(UClothingAssetBase::InvalidateAllCachedData(),);
 #endif // WITH_EDITORONLY_DATA
diff --git a/Engine/Source/Runtime/Engine/Private/SkeletalMesh.cpp b/Engine/Source/Runtime/Engine/Private/SkeletalMesh.cpp
--- a/Engine/Source/Runtime/Engine/Private/SkeletalMesh.cpp
+++ b/Engine/Source/Runtime/Engine/Private/SkeletalMesh.cpp
@@ -4009,18 +4009,6 @@ void USkeletalMesh::FinishPostLoadInternal(FSkinnedAssetPostLoadContext& Context
 		// Convert screen area to screen size
 		ConvertLegacyLODScreenSize();
 	}
-
-	// If inverse masses have never been cached, invalidate data so it will be recalculated
-	if (GetLinkerCustomVersion(FSkeletalMeshCustomVersion::GUID) < FSkeletalMeshCustomVersion::CachedClothInverseMasses)
-	{
-		for (UClothingAssetBase* ClothingAsset : GetMeshClothingAssets())
-		{
-			if (ClothingAsset) 
-			{
-				ClothingAsset->InvalidateAllCachedData();
-			}
-		}
-	}
 #endif
 
 	SetHasActiveClothingAssets(ComputeActiveClothingAssets());
diff --git a/Engine/Source/Runtime/Engine/Public/SkeletalMeshLegacyCustomVersions.h b/Engine/Source/Runtime/Engine/Public/SkeletalMeshLegacyCustomVersions.h
--- a/Engine/Source/Runtime/Engine/Public/SkeletalMeshLegacyCustomVersions.h
+++ b/Engine/Source/Runtime/Engine/Public/SkeletalMeshLegacyCustomVersions.h
@@ -55,7 +55,7 @@ struct FSkeletalMeshCustomVersion
 	};
 
 	// The GUID for this custom version number
-	const static FGuid GUID;
+	ENGINE_API const static FGuid GUID;
 
 private:
 	FSkeletalMeshCustomVersion() {}
