Cloth - Added an AppendUnique() method to FClothCollisionData.

This will be soon required in places where collisions are extracted from the multiple simulation instances running on a SkeletalMeshComponent (review soon to come).

See:
FEnvironmentalCollisions::AppendCollisionDataFromEnvironment() where cloth collisions are extracted from nearby skeletal meshes, and USkeletalMeshComponent::FindClothCollisions() which returns the used cloth collisions.

#rnx

diff --git a/Engine/Source/Runtime/ClothingSystemRuntimeInterface/Private/ClothCollisionData.cpp b/Engine/Source/Runtime/ClothingSystemRuntimeInterface/Private/ClothCollisionData.cpp
--- a/Engine/Source/Runtime/ClothingSystemRuntimeInterface/Private/ClothCollisionData.cpp
+++ b/Engine/Source/Runtime/ClothingSystemRuntimeInterface/Private/ClothCollisionData.cpp
@@ -66,3 +66,40 @@ void FClothCollisionData::AppendTransformed(const FClothCollisionData& InOther,
 		Boxes[Index].LocalRotation = Transform.TransformRotation(Boxes[Index].LocalRotation);
 	}
 }
+
+void FClothCollisionData::AppendUnique(const FClothCollisionData& InOther)
+{
+	TArray<int32> OldToNewSphereIndices;
+	OldToNewSphereIndices.SetNumUninitialized(InOther.Spheres.Num());
+
+	Spheres.Reserve(Spheres.Num() + InOther.Spheres.Num());
+	for (int32 SphereIndex = 0; SphereIndex < InOther.Spheres.Num(); ++SphereIndex)
+	{
+		OldToNewSphereIndices[SphereIndex] = Spheres.AddUnique(InOther.Spheres[SphereIndex]);
+	}
+
+	SphereConnections.Reserve(SphereConnections.Num() + InOther.SphereConnections.Num());
+	for (int32 SphereConnectionIndex = 0; SphereConnectionIndex < InOther.SphereConnections.Num(); ++SphereConnectionIndex)
+	{
+		FClothCollisionPrim_SphereConnection SphereConnection = SphereConnections[SphereConnectionIndex];
+		if (OldToNewSphereIndices.IsValidIndex(SphereConnection.SphereIndices[0]) &&
+			OldToNewSphereIndices.IsValidIndex(SphereConnection.SphereIndices[1]))
+		{
+			SphereConnection.SphereIndices[0] = OldToNewSphereIndices[SphereConnection.SphereIndices[0]];
+			SphereConnection.SphereIndices[1] = OldToNewSphereIndices[SphereConnection.SphereIndices[1]];
+		}
+		SphereConnections.AddUnique(MoveTemp(SphereConnection));
+	}
+
+	Convexes.Reserve(Convexes.Num() + InOther.Convexes.Num());
+	for (const FClothCollisionPrim_Convex& Convex : InOther.Convexes)
+	{
+		Convexes.AddUnique(Convex);
+	}
+
+	Boxes.Reserve(Boxes.Num() + InOther.Boxes.Num());
+	for (const FClothCollisionPrim_Box& Box : InOther.Boxes)
+	{
+		Boxes.AddUnique(Box);
+	}
+}
diff --git a/Engine/Source/Runtime/ClothingSystemRuntimeInterface/Public/ClothCollisionData.h b/Engine/Source/Runtime/ClothingSystemRuntimeInterface/Public/ClothCollisionData.h
--- a/Engine/Source/Runtime/ClothingSystemRuntimeInterface/Public/ClothCollisionData.h
+++ b/Engine/Source/Runtime/ClothingSystemRuntimeInterface/Public/ClothCollisionData.h
@@ -7,16 +7,24 @@
 #include "UObject/ObjectMacros.h"
 #include "ClothCollisionData.generated.h"
 
+#define UE_API CLOTHINGSYSTEMRUNTIMEINTERFACE_API
+
 USTRUCT()
 struct FClothCollisionData
 {
 	GENERATED_BODY()
 
-	CLOTHINGSYSTEMRUNTIMEINTERFACE_API void Reset();
+	UE_API void Reset();
+
+	UE_API void Append(const FClothCollisionData& InOther);
 
-	CLOTHINGSYSTEMRUNTIMEINTERFACE_API void Append(const FClothCollisionData& InOther);
+	UE_API void AppendTransformed(const FClothCollisionData& InOther, const TArray<FTransform>& BoneTransforms);
 
-	CLOTHINGSYSTEMRUNTIMEINTERFACE_API void AppendTransformed(const FClothCollisionData& InOther, const TArray<FTransform>& BoneTransforms);
+	/**
+	 * Append collision data, but only the individual collision elements that don't already exist as to avoid duplicates.
+	 * @Note This is an expensive operation and therefore it is best to call Append() on the first batch of collision data to merge right after a Reset().
+	 */
+	UE_API void AppendUnique(const FClothCollisionData& InOther);
 
 	bool IsEmpty() const
 	{
@@ -24,18 +32,20 @@ struct FClothCollisionData
 	}
 
 	// Sphere data
-	UPROPERTY(EditAnywhere, Category = Collison)
+	UPROPERTY(EditAnywhere, Category = Collision)
 	TArray<FClothCollisionPrim_Sphere> Spheres;
 
 	// Capsule data
-	UPROPERTY(EditAnywhere, Category = Collison)
+	UPROPERTY(EditAnywhere, Category = Collision)
 	TArray<FClothCollisionPrim_SphereConnection> SphereConnections;
 
 	// Convex Data
-	UPROPERTY(EditAnywhere, Category = Collison)
+	UPROPERTY(EditAnywhere, Category = Collision)
 	TArray<FClothCollisionPrim_Convex> Convexes;
 
 	// Box data
-	UPROPERTY(EditAnywhere, Category = Collison)
+	UPROPERTY(EditAnywhere, Category = Collision)
 	TArray<FClothCollisionPrim_Box> Boxes;
 };
+
+#undef UE_API
diff --git a/Engine/Source/Runtime/ClothingSystemRuntimeInterface/Public/ClothCollisionPrim.h b/Engine/Source/Runtime/ClothingSystemRuntimeInterface/Public/ClothCollisionPrim.h
--- a/Engine/Source/Runtime/ClothingSystemRuntimeInterface/Public/ClothCollisionPrim.h
+++ b/Engine/Source/Runtime/ClothingSystemRuntimeInterface/Public/ClothCollisionPrim.h
@@ -20,6 +20,14 @@ struct FClothCollisionPrim_Sphere
 		, LocalPosition(InLocalPosition)
 	{}
 
+	bool operator==(const FClothCollisionPrim_Sphere& Other) const
+	{
+		return
+			BoneIndex == Other.BoneIndex &&
+			FMath::IsNearlyEqual(Radius, Other.Radius) &&
+			LocalPosition.Equals(Other.LocalPosition);
+	}
+
 	UPROPERTY()
 	int32 BoneIndex;
 
@@ -45,6 +53,11 @@ struct FClothCollisionPrim_SphereConnection
 		OneSidedPlaneNormal = InOneSidedPlaneNormal;
 	}
 
+	bool operator==(const FClothCollisionPrim_SphereConnection& Other) const
+	{
+		return SphereIndices[0] == SphereIndices[0] && SphereIndices[1] == SphereIndices[1] && OneSidedPlaneNormal == Other.OneSidedPlaneNormal;
+	}
+
 	UPROPERTY()
 	int32 SphereIndices[2];
 
@@ -61,6 +74,11 @@ struct FClothCollisionPrim_ConvexFace
 
 	FClothCollisionPrim_ConvexFace(): Plane(ForceInit) {}
 
+	bool operator==(const FClothCollisionPrim_ConvexFace& Other) const
+	{
+		return Plane.Equals(Other.Plane) && Indices == Other.Indices;
+	}
+
 	UPROPERTY()
 	FPlane Plane;
 
@@ -91,6 +109,24 @@ struct FClothCollisionPrim_Convex
 		, BoneIndex(InBoneIndex)
 	{}
 
+	bool operator==(const FClothCollisionPrim_Convex& Other) const
+	{
+		if (BoneIndex == Other.BoneIndex &&
+			Faces == Other.Faces &&
+			SurfacePoints.Num() == Other.SurfacePoints.Num())
+		{
+			for (int32 Index = 0; Index < SurfacePoints.Num(); ++Index)
+			{
+				if (!SurfacePoints[Index].Equals(Other.SurfacePoints[Index]))
+				{
+					return false;
+				}
+			}
+			return true;
+		}
+		return false;
+	}
+
 #if WITH_EDITORONLY_DATA
 	UPROPERTY()
 	TArray<FPlane> Planes_DEPRECATED;
@@ -123,6 +159,15 @@ struct FClothCollisionPrim_Box
 		, BoneIndex(InBoneIndex)
 	{}
 
+	bool operator==(const FClothCollisionPrim_Box& Other) const
+	{
+		return
+			BoneIndex == Other.BoneIndex &&
+			LocalPosition.Equals(Other.LocalPosition) &&
+			LocalRotation.Equals(Other.LocalRotation) &&
+			HalfExtents.Equals(Other.HalfExtents);
+	}
+
 	UPROPERTY()
 	FVector LocalPosition;
 
