CIS - Fixed Issue 1007962: PVS Studio warnings in ClothCollisionPrim.h

#rnx

diff --git a/Engine/Source/Runtime/ClothingSystemRuntimeInterface/Private/ClothCollisionData.cpp b/Engine/Source/Runtime/ClothingSystemRuntimeInterface/Private/ClothCollisionData.cpp
--- a/Engine/Source/Runtime/ClothingSystemRuntimeInterface/Private/ClothCollisionData.cpp
+++ b/Engine/Source/Runtime/ClothingSystemRuntimeInterface/Private/ClothCollisionData.cpp
@@ -81,7 +81,7 @@ void FClothCollisionData::AppendUnique(const FClothCollisionData& InOther)
 	SphereConnections.Reserve(SphereConnections.Num() + InOther.SphereConnections.Num());
 	for (int32 SphereConnectionIndex = 0; SphereConnectionIndex < InOther.SphereConnections.Num(); ++SphereConnectionIndex)
 	{
-		FClothCollisionPrim_SphereConnection SphereConnection = SphereConnections[SphereConnectionIndex];
+		FClothCollisionPrim_SphereConnection SphereConnection = InOther.SphereConnections[SphereConnectionIndex];
 		if (OldToNewSphereIndices.IsValidIndex(SphereConnection.SphereIndices[0]) &&
 			OldToNewSphereIndices.IsValidIndex(SphereConnection.SphereIndices[1]))
 		{
diff --git a/Engine/Source/Runtime/ClothingSystemRuntimeInterface/Public/ClothCollisionPrim.h b/Engine/Source/Runtime/ClothingSystemRuntimeInterface/Public/ClothCollisionPrim.h
--- a/Engine/Source/Runtime/ClothingSystemRuntimeInterface/Public/ClothCollisionPrim.h
+++ b/Engine/Source/Runtime/ClothingSystemRuntimeInterface/Public/ClothCollisionPrim.h
@@ -55,7 +55,10 @@ struct FClothCollisionPrim_SphereConnection
 
 	bool operator==(const FClothCollisionPrim_SphereConnection& Other) const
 	{
-		return SphereIndices[0] == SphereIndices[0] && SphereIndices[1] == SphereIndices[1] && OneSidedPlaneNormal == Other.OneSidedPlaneNormal;
+		return
+			SphereIndices[0] == Other.SphereIndices[0] &&
+			SphereIndices[1] == Other.SphereIndices[1] &&
+			OneSidedPlaneNormal == Other.OneSidedPlaneNormal;
 	}
 
 	UPROPERTY()
