Chaos Cloth Asset - Various small fixes and improvements after the SKM Cloth Asset changes CL 45029949:
- Fixed UnbindClothingAndBackup in the Cloth Asset SKM clothing data code.
- Fixed static_cast to possibly the wrong simulation type in the SkeletalMeshComponentCacheAdapter.
- Fixed ClothAssetSKMClothingSimulation GetBounds().
- Fixed FClothingSimulationSkeletalMesh::GetBoneTransforms() when context is null.
- Fixed default implementations of depreacated const functions in ClothingSimulationInterface.h.

#rnx

diff --git a/Engine/Plugins/ChaosCloth/Source/ChaosCloth/Private/ChaosCloth/ChaosClothingSimulationSkeletalMesh.cpp b/Engine/Plugins/ChaosCloth/Source/ChaosCloth/Private/ChaosCloth/ChaosClothingSimulationSkeletalMesh.cpp
--- a/Engine/Plugins/ChaosCloth/Source/ChaosCloth/Private/ChaosCloth/ChaosClothingSimulationSkeletalMesh.cpp
+++ b/Engine/Plugins/ChaosCloth/Source/ChaosCloth/Private/ChaosCloth/ChaosClothingSimulationSkeletalMesh.cpp
@@ -234,7 +234,8 @@ namespace Chaos
 
 	const TArray<FTransform>& FClothingSimulationSkeletalMesh::GetBoneTransforms() const
 	{
-		return GetContext()->BoneTransforms;
+		static TArray<FTransform> EmptyArray;
+		return GetContext() ? GetContext()->BoneTransforms : EmptyArray;
 	}
 
 	const FTransform& FClothingSimulationSkeletalMesh::GetComponentToWorldTransform() const
diff --git a/Engine/Plugins/ChaosCloth/Source/ChaosCloth/Private/ChaosCloth/SkeletalMeshComponentCacheAdapter.cpp b/Engine/Plugins/ChaosCloth/Source/ChaosCloth/Private/ChaosCloth/SkeletalMeshComponentCacheAdapter.cpp
--- a/Engine/Plugins/ChaosCloth/Source/ChaosCloth/Private/ChaosCloth/SkeletalMeshComponentCacheAdapter.cpp
+++ b/Engine/Plugins/ChaosCloth/Source/ChaosCloth/Private/ChaosCloth/SkeletalMeshComponentCacheAdapter.cpp
@@ -2,6 +2,7 @@
 
 #include "ChaosCloth/SkeletalMeshComponentCacheAdapter.h"
 #include "ChaosCloth/ChaosClothingCacheSchema.h"
+#include "ChaosCloth/ChaosClothingSimulationFactory.h"
 #include "ChaosCloth/ChaosClothPrivate.h"
 #include "Chaos/ParticleHandle.h"
 #include "Chaos/ChaosCache.h"
@@ -9,6 +10,7 @@
 #include "ChaosCloth/ChaosClothingSimulation.h"
 #include "ChaosCloth/ChaosClothingSimulationSolver.h"
 #include "Components/SkeletalMeshComponent.h"
+#include "ClothingSimulationInstance.h"
 
 namespace Chaos
 {
@@ -81,12 +83,18 @@ namespace Chaos
 						TEXT("SkeletalMeshComponent [%s] has more than one clothing simulation running, but cloth caching can only deal with a single one."), *MeshComp->GetName());
 					bErroredOnce = true;
 				}
-PRAGMA_DISABLE_DEPRECATION_WARNINGS
-				if(MeshComp->GetClothingSimulation())
+
+				for (FClothingSimulationInstance& ClothingSimulationInstance : MeshComp->GetClothingSimulationInstances())
 				{
-					return StaticCast<FClothingSimulation*>(MeshComp->GetClothingSimulation())->Solver.Get();
+					if (ClothingSimulationInstance.GetClothingSimulationFactory())
+					{
+						// TODO: Manage multiple solvers and solver types as to process all instances
+						if (ClothingSimulationInstance.GetClothingSimulationFactory()->GetClass() == UChaosClothingSimulationFactory::StaticClass())
+						{
+							return StaticCast<FClothingSimulation*>(ClothingSimulationInstance.GetClothingSimulation())->Solver.Get();
+						}
+					}
 				}
-PRAGMA_ENABLE_DEPRECATION_WARNINGS
 			}
 		}
 		return nullptr;
diff --git a/Engine/Plugins/ChaosClothAsset/Source/ChaosClothAssetEngine/Private/ChaosClothAsset/ClothAssetSKMClothingAsset.cpp b/Engine/Plugins/ChaosClothAsset/Source/ChaosClothAssetEngine/Private/ChaosClothAsset/ClothAssetSKMClothingAsset.cpp
--- a/Engine/Plugins/ChaosClothAsset/Source/ChaosClothAssetEngine/Private/ChaosClothAsset/ClothAssetSKMClothingAsset.cpp
+++ b/Engine/Plugins/ChaosClothAsset/Source/ChaosClothAssetEngine/Private/ChaosClothAsset/ClothAssetSKMClothingAsset.cpp
@@ -29,13 +29,13 @@ namespace UE::Chaos::ClothAsset::Private
 
 #if WITH_EDITOR
 	// Modified version of FLODUtilities::UnbindClothingAndBackup to only unbind specified assets
-	static void UnbindClothingAndBackup(USkeletalMesh& SkeletalMesh, const UChaosClothAssetSKMClothingAsset& ClothingAsset, TArray<ClothingAssetUtils::FClothingAssetMeshBinding> ClothingBindings)
+	static void UnbindClothingAndBackup(USkeletalMesh& SkeletalMesh, const UChaosClothAssetSKMClothingAsset& ClothingAsset, TArray<ClothingAssetUtils::FClothingAssetMeshBinding>& ClothingBindings)
 	{
 		ClothingBindings.Reset();
 
 		if (SkeletalMesh.GetImportedModel())
 		{
-			auto UnbindClothingLODAndBackup = [&SkeletalMesh, &ClothingAsset](const int32 LODIndex, TArray<ClothingAssetUtils::FClothingAssetMeshBinding> ClothingLODBindings)
+			auto UnbindClothingLODAndBackup = [&SkeletalMesh, &ClothingAsset](const int32 LODIndex, TArray<ClothingAssetUtils::FClothingAssetMeshBinding>& ClothingLODBindings)
 				{
 					ClothingLODBindings.Reset();
 
@@ -74,7 +74,7 @@ namespace UE::Chaos::ClothAsset::Private
 		}
 	}
 
-	// Modified version of FLODUtilities::RestoreClothingFromBackup to only unbind specified assets
+	// Modified version of FLODUtilities::RestoreClothingFromBackup to only rebind specified assets
 	static void RestoreClothingFromBackup(USkeletalMesh& SkeletalMesh, const UChaosClothAssetSKMClothingAsset& ClothingAsset, TArray<ClothingAssetUtils::FClothingAssetMeshBinding>& ClothingBindings)
 	{
 		if (SkeletalMesh.GetImportedModel())
diff --git a/Engine/Plugins/ChaosClothAssetEditor/Source/ChaosClothAssetTools/Private/ChaosClothAsset/ChaosClothAssetToolsModule.cpp b/Engine/Plugins/ChaosClothAssetEditor/Source/ChaosClothAssetTools/Private/ChaosClothAsset/ChaosClothAssetToolsModule.cpp
--- a/Engine/Plugins/ChaosClothAssetEditor/Source/ChaosClothAssetTools/Private/ChaosClothAsset/ChaosClothAssetToolsModule.cpp
+++ b/Engine/Plugins/ChaosClothAssetEditor/Source/ChaosClothAssetTools/Private/ChaosClothAsset/ChaosClothAssetToolsModule.cpp
@@ -61,7 +61,7 @@ namespace UE::Chaos::ClothAsset
 		{
 			return UChaosClothAssetSKMClothingAssetFactory::StaticClass()->GetDefaultObject<UClothingAssetFactoryBase>();
 		}
-		//~ End IClothingSimulationFactoryClassProvider
+		//~ End IClothingAssetFactoryProvider
 	};
 }
 
diff --git a/Engine/Source/Runtime/ClothingSystemRuntimeCommon/Public/ClothingAsset.h b/Engine/Source/Runtime/ClothingSystemRuntimeCommon/Public/ClothingAsset.h
--- a/Engine/Source/Runtime/ClothingSystemRuntimeCommon/Public/ClothingAsset.h
+++ b/Engine/Source/Runtime/ClothingSystemRuntimeCommon/Public/ClothingAsset.h
@@ -33,6 +33,7 @@ namespace ClothingAssetUtils
 #if WITH_EDITOR
 	/**
 	 * Unbind this clothing asset from the provided skeletal mesh, will remove all LODs.
+	 * @param Asset The clothing data to unbind.
 	 * @param SkeletalMesh skeletal mesh to remove from.
 	 */
 	UE_API void UnbindFromSkeletalMesh(const UClothingAssetBase& Asset, USkeletalMesh& SkeletalMesh);
diff --git a/Engine/Source/Runtime/ClothingSystemRuntimeInterface/Public/ClothingSimulationFactory.h b/Engine/Source/Runtime/ClothingSystemRuntimeInterface/Public/ClothingSimulationFactory.h
--- a/Engine/Source/Runtime/ClothingSystemRuntimeInterface/Public/ClothingSimulationFactory.h
+++ b/Engine/Source/Runtime/ClothingSystemRuntimeInterface/Public/ClothingSimulationFactory.h
@@ -51,9 +51,9 @@ public:
 	// Otherwise return the last registered factory.
 	static UE_API TSubclassOf<class UClothingSimulationFactory> GetDefaultClothingSimulationFactoryClass();
 
-	// Return the clothing simulation factory class for.the specified asset.
+	// Return the clothing simulation factory for the specified asset.
 	// Calls SupportsAsset() internally.
-	// If multiple simulation classes support this asset then this returns the first one registered through IClothingSimulationFactoryClassProvider.
+	// If multiple simulation support the specified asset type then this returns the first one registered through IClothingSimulationFactoryClassProvider.
 	static UE_API UClothingSimulationFactory* GetClothingSimulationFactory(const UClothingAssetBase* InAsset);
 
 	// Create a simulation object for a skeletal mesh to use (see IClothingSimulation)
diff --git a/Engine/Source/Runtime/ClothingSystemRuntimeInterface/Public/ClothingSimulationInterface.h b/Engine/Source/Runtime/ClothingSystemRuntimeInterface/Public/ClothingSimulationInterface.h
--- a/Engine/Source/Runtime/ClothingSystemRuntimeInterface/Public/ClothingSimulationInterface.h
+++ b/Engine/Source/Runtime/ClothingSystemRuntimeInterface/Public/ClothingSimulationInterface.h
@@ -50,14 +50,16 @@ protected:
 	 * @param InAsset - The asset to create an actor from
 	 * @param SimDataIndex - the sim data index to use when doing the writeback in GetSimulationData
 	 */
-	virtual void CreateActor(USkeletalMeshComponent* InOwnerComponent, const UClothingAssetBase* InAsset, int32 SimDataIndex) {}
-
-	UE_DEPRECATED(5.7, "Use const arguments version instead,")
-	virtual void CreateActor(USkeletalMeshComponent* InOwnerComponent, UClothingAssetBase* InAsset, int32 SimDataIndex)
+	virtual void CreateActor(USkeletalMeshComponent* InOwnerComponent, const UClothingAssetBase* InAsset, int32 SimDataIndex)
 	{
-		return CreateActor(InOwnerComponent, static_cast<const UClothingAssetBase*>(InAsset), SimDataIndex);
+PRAGMA_DISABLE_DEPRECATION_WARNINGS
+		return CreateActor(InOwnerComponent, const_cast<UClothingAssetBase*>(InAsset), SimDataIndex);
+PRAGMA_ENABLE_DEPRECATION_WARNINGS
 	}
 
+	UE_DEPRECATED(5.7, "Use const arguments version instead,")
+	virtual void CreateActor(USkeletalMeshComponent* InOwnerComponent, UClothingAssetBase* InAsset, int32 SimDataIndex) {}
+
 	/** Called once all CreateActor have been called. */
 	virtual void EndCreateActor() {}
 
@@ -70,18 +72,22 @@ protected:
 	 * @param InOutContext - The context to fill
 	 * @param bIsInitialization - Whether this fill is occurring as part of the actor creation stage
 	 */
-	virtual void FillContext(const USkeletalMeshComponent* InComponent, float InDeltaTime, IClothingSimulationContext* InOutContext, bool bIsInitialization) const {}
+	virtual void FillContext(const USkeletalMeshComponent* InComponent, float InDeltaTime, IClothingSimulationContext* InOutContext, bool bIsInitialization) const
+	{
+PRAGMA_DISABLE_DEPRECATION_WARNINGS
+		return const_cast<IClothingSimulation*>(this)->FillContext(InComponent, InDeltaTime, InOutContext, bIsInitialization);
+PRAGMA_ENABLE_DEPRECATION_WARNINGS
+	}
 
 	UE_DEPRECATED(5.7, "Use const version instead.")
 	virtual void FillContext(USkeletalMeshComponent* InComponent, float InDeltaTime, IClothingSimulationContext* InOutContext, bool bIsInitialization)
 	{
-		return static_cast<const IClothingSimulation*>(this)->FillContext(InComponent, InDeltaTime, InOutContext, bIsInitialization);
+PRAGMA_DISABLE_DEPRECATION_WARNINGS
+		return FillContext(InComponent, InDeltaTime, InOutContext);
+PRAGMA_ENABLE_DEPRECATION_WARNINGS
 	}
 	UE_DEPRECATED(5.7, "Use const version instead.")
-	virtual void FillContext(USkeletalMeshComponent* InComponent, float InDeltaTime, IClothingSimulationContext* InOutContext)
-	{
-		return static_cast<const IClothingSimulation*>(this)->FillContext(InComponent, InDeltaTime, InOutContext, false);
-	}
+	virtual void FillContext(USkeletalMeshComponent* InComponent, float InDeltaTime, IClothingSimulationContext* InOutContext) {}
 
 	/** Initialize the simulation, will be called before any Simulate calls */
 	virtual void Initialize() = 0;
@@ -105,13 +111,15 @@ protected:
 	 * New callers should make care to make sure only the correct context is ever passed through
 	 * @param InContext - The context to use during simulation, will have been filled in FillContext
 	 */
-	virtual void Simulate(const IClothingSimulationContext* InContext) {}
-	
-	UE_DEPRECATED(5.7, "Use const argument version instead,")
-	virtual void Simulate(IClothingSimulationContext* InContext)
+	virtual void Simulate(const IClothingSimulationContext* InContext)
 	{
-		return Simulate(static_cast<const IClothingSimulationContext*>(InContext));
+PRAGMA_DISABLE_DEPRECATION_WARNINGS
+		return Simulate(const_cast<IClothingSimulationContext*>(InContext));
+PRAGMA_ENABLE_DEPRECATION_WARNINGS
 	}
+	
+	UE_DEPRECATED(5.7, "Use const argument version instead,")
+	virtual void Simulate(IClothingSimulationContext* InContext) {}
 
 	/** Simulation should remove all of it's actors when next possible and free them */
 	virtual void DestroyActors() = 0;
@@ -128,14 +136,16 @@ protected:
 	 * @param OwnerComponent - the component that owns the simulation
 	 * @param OverrideComponent - An override component if bound to a leader pose component
 	 */
-	virtual void GetSimulationData(TMap<int32, FClothSimulData>& OutData, const USkeletalMeshComponent* OwnerComponent, const USkinnedMeshComponent* OverrideComponent) const {}
-
-	UE_DEPRECATED(5.7, "Use const version instead.")
-	virtual void GetSimulationData(TMap<int32, FClothSimulData>& OutData, USkeletalMeshComponent* InOwnerComponent, USkinnedMeshComponent* InOverrideComponent) const
+	virtual void GetSimulationData(TMap<int32, FClothSimulData>& OutData, const USkeletalMeshComponent* OwnerComponent, const USkinnedMeshComponent* OverrideComponent) const
 	{
-		return GetSimulationData(OutData, static_cast<const USkeletalMeshComponent*>(InOwnerComponent), static_cast<const USkinnedMeshComponent*>(InOverrideComponent));
+PRAGMA_DISABLE_DEPRECATION_WARNINGS
+		return GetSimulationData(OutData, const_cast<USkeletalMeshComponent*>(OwnerComponent), const_cast<USkinnedMeshComponent*>(OverrideComponent));
+PRAGMA_ENABLE_DEPRECATION_WARNINGS
 	}
 
+	UE_DEPRECATED(5.7, "Use const version instead.")
+	virtual void GetSimulationData(TMap<int32, FClothSimulData>& OutData, USkeletalMeshComponent* InOwnerComponent, USkinnedMeshComponent* InOverrideComponent) const {}
+
 	/** Get the bounds of the simulation mesh in local simulation space */
 	virtual FBoxSphereBounds GetBounds(const USkeletalMeshComponent* InOwnerComponent) const = 0;
 
