Cloth - Deprecated the UClothingSimulationInteractor::GatherStats() method which has never been used since added.

#rnx

diff --git a/Engine/Source/Runtime/ClothingSystemRuntimeInterface/Public/ClothingSimulationInterface.h b/Engine/Source/Runtime/ClothingSystemRuntimeInterface/Public/ClothingSimulationInterface.h
--- a/Engine/Source/Runtime/ClothingSystemRuntimeInterface/Public/ClothingSimulationInterface.h
+++ b/Engine/Source/Runtime/ClothingSystemRuntimeInterface/Public/ClothingSimulationInterface.h
@@ -190,14 +190,7 @@ public:
 	 */
 	virtual void SetNumSubsteps(int32 /*NumSubsteps*/) {}
 
-	/**
-	 * Called by the engine once per physics tick to gather non-cycle stats (num sim verts etc.)
-	 * Will not be called if STATS is not defined, not required to be implemented if no need for stats.
-	 * STAT_NumCloths and STAT_NumClothVerts are already defined and can be incremented, more can be
-	 * added per simulation.
-	 * This can be called on any thread, so derived implementations should make sure the gather is
-	 * safe
-	 */
+	UE_DEPRECATED(5.7, "GatherStats is unused and will be removed from future interfaces. Do not override.")
 	virtual void GatherStats() const {}
 
 	/** Return the number of simulated cloths. Implementation must be thread safe. */
