Cloth:
- Changed GetClothingInteractor(const FString&) to GetClothingInteractor(const FName).
- Deprecated GetClothingInteractor(const FString&).
- Made the CreateClothingInteractor method a virtual method so that complete control of the interactor creation can be done in the derived classes.

#rnx

diff --git a/Engine/Source/Runtime/ClothingSystemRuntimeInterface/Private/ClothingSimulationInteractor.cpp b/Engine/Source/Runtime/ClothingSystemRuntimeInterface/Private/ClothingSimulationInteractor.cpp
--- a/Engine/Source/Runtime/ClothingSystemRuntimeInterface/Private/ClothingSimulationInteractor.cpp
+++ b/Engine/Source/Runtime/ClothingSystemRuntimeInterface/Private/ClothingSimulationInteractor.cpp
@@ -24,9 +24,9 @@ void UClothingSimulationInteractor::DestroyClothingInteractors()
 	ClothingInteractors.Reset();
 }
 
-UClothingInteractor* UClothingSimulationInteractor::GetClothingInteractor(const FString& ClothingAssetName) const
+UClothingInteractor* UClothingSimulationInteractor::GetClothingInteractor(const FName ClothingAssetName) const
 {
-	if (const TObjectPtr<UClothingInteractor>* ClothingInteractor = ClothingInteractors.Find(FName(ClothingAssetName)))
+	if (const TObjectPtr<UClothingInteractor>* ClothingInteractor = ClothingInteractors.Find(ClothingAssetName))
 	{
 		return *ClothingInteractor;
 	}
diff --git a/Engine/Source/Runtime/ClothingSystemRuntimeInterface/Public/ClothingSimulationInteractor.h b/Engine/Source/Runtime/ClothingSystemRuntimeInterface/Public/ClothingSimulationInteractor.h
--- a/Engine/Source/Runtime/ClothingSystemRuntimeInterface/Public/ClothingSimulationInteractor.h
+++ b/Engine/Source/Runtime/ClothingSystemRuntimeInterface/Public/ClothingSimulationInteractor.h
@@ -51,8 +51,8 @@ class UClothingSimulationInteractor : public UObject
 	GENERATED_BODY()
 
 public:
-	/** Create a cloth specific interactor. */
-	CLOTHINGSYSTEMRUNTIMEINTERFACE_API void CreateClothingInteractor(const UClothingAssetBase* ClothingAsset, int32 ClothingId);
+	/** Create a cloth specific interactor. Override either this method or the simpler CreateClothingInteractor method. */
+	CLOTHINGSYSTEMRUNTIMEINTERFACE_API virtual void CreateClothingInteractor(const UClothingAssetBase* ClothingAsset, int32 ClothingId);
 
 	/** Destroy all interactors. */
 	CLOTHINGSYSTEMRUNTIMEINTERFACE_API void DestroyClothingInteractors();
@@ -142,15 +142,24 @@ public:
 
 	/**Return a cloth interactor for this simulation. */
 	UFUNCTION(BlueprintCallable, Category = ClothingSimulation)
-	CLOTHINGSYSTEMRUNTIMEINTERFACE_API UClothingInteractor* GetClothingInteractor(const FString& ClothingAssetName) const;
+	CLOTHINGSYSTEMRUNTIMEINTERFACE_API UClothingInteractor* GetClothingInteractor(const FName ClothingAssetName) const;
+
+	UE_DEPRECATED(5.7, "Use FName version instead.")
+	CLOTHINGSYSTEMRUNTIMEINTERFACE_API UClothingInteractor* GetClothingInteractor(const FString& ClothingAssetName) const
+	{
+		return GetClothingInteractor(FName(ClothingAssetName));
+	}
 
 	/** Cloth interactors currently created. */
 	UPROPERTY()
 	TMap<FName, TObjectPtr<UClothingInteractor>> ClothingInteractors;
 
 protected:
+	/** Create a clothing interactor. This method is called by the public CreateClothingInteractor() default implementation. */
 	virtual UClothingInteractor* CreateClothingInteractor()
-	PURE_VIRTUAL(UClothingSimulationInteractor::CreateClothingInteractor, return nullptr;);
+	{
+		return nullptr;
+	}
 
 private:
 	// Simulation infos updated during the sync
