Include Math/BoxSphereBounds.h

diff --git a/Engine/Source/Runtime/ClothingSystemRuntimeInterface/Public/ClothingSimulationInstance.h b/Engine/Source/Runtime/ClothingSystemRuntimeInterface/Public/ClothingSimulationInstance.h
--- a/Engine/Source/Runtime/ClothingSystemRuntimeInterface/Public/ClothingSimulationInstance.h
+++ b/Engine/Source/Runtime/ClothingSystemRuntimeInterface/Public/ClothingSimulationInstance.h
@@ -2,6 +2,7 @@
 #pragma once
 
 #include "ClothingSimulationInterface.h"
+#include "Math/BoxSphereBounds.h"
 #include "UObject/ObjectMacros.h"
 #include "UObject/ObjectPtr.h"
 #include "ClothingSimulationInstance.generated.h"
