Cloth - Small trivial improvement to the simulation instance class constructor.

#rnx

diff --git a/Engine/Source/Runtime/ClothingSystemRuntimeInterface/Private/ClothingSimulationInstance.cpp b/Engine/Source/Runtime/ClothingSystemRuntimeInterface/Private/ClothingSimulationInstance.cpp
--- a/Engine/Source/Runtime/ClothingSystemRuntimeInterface/Private/ClothingSimulationInstance.cpp
+++ b/Engine/Source/Runtime/ClothingSystemRuntimeInterface/Private/ClothingSimulationInstance.cpp
@@ -10,7 +10,7 @@
 
 FClothingSimulationInstance::FClothingSimulationInstance(USkeletalMeshComponent* SkeletalMeshComponent, const FFactoryAssetGroup& FactoryAssetGroup)
 	: ClothingSimulation(static_cast<const UClothingSimulationFactory*>(FactoryAssetGroup.ClothingSimulationFactory)->CreateSimulation())  // UE_DEPRECATED(5.7, "Use IClothingSimulationInterface instead.")
-	, ClothingSimulationContext(ClothingSimulation->CreateContext())
+	, ClothingSimulationContext(ClothingSimulation ? ClothingSimulation->CreateContext() : nullptr)
 	, ClothingSimulationInteractor(static_cast<const UClothingSimulationFactory*>(FactoryAssetGroup.ClothingSimulationFactory)->SupportsRuntimeInteraction() ? FactoryAssetGroup.ClothingSimulationFactory->CreateInteractor() : nullptr)
 	, ClothingSimulationFactory(FactoryAssetGroup.ClothingSimulationFactory)
 {
