Cloth - Minor fix in the clothing interface doc comments (Append instead of GetSimulationData) .

#rnx

diff --git a/Engine/Source/Runtime/ClothingSystemRuntimeInterface/Public/ClothingSimulationInterface.h b/Engine/Source/Runtime/ClothingSystemRuntimeInterface/Public/ClothingSimulationInterface.h
--- a/Engine/Source/Runtime/ClothingSystemRuntimeInterface/Public/ClothingSimulationInterface.h
+++ b/Engine/Source/Runtime/ClothingSystemRuntimeInterface/Public/ClothingSimulationInterface.h
@@ -43,7 +43,7 @@ protected:
 
 	/**
 	 * Create an actor for this simulation from the data in InAsset
-	 * Simulation data for this actor should be written back to SimDataIndex in GetSimulationData
+	 * Simulation data for this actor should be written back to SimDataIndex in AppendSimulationData
 	 * @param InOwnerComponent - The component requesting this actor
 	 * @param InAsset - The asset to create an actor from
 	 * @param SimDataIndex - the sim data index to use when doing the writeback in AppendSimulationData
@@ -207,10 +207,10 @@ protected:
 
 	/**
 	 * Create an actor for this simulation from the data in InAsset
-	 * Simulation data for this actor should be written back to SimDataIndex in GetSimulationData
+	 * Simulation data for this actor should be written back to SimDataIndex in AppendSimulationData
 	 * @param InOwnerComponent - The component requesting this actor
 	 * @param InAsset - The asset to create an actor from
-	 * @param SimDataIndex - the sim data index to use when doing the writeback in GetSimulationData
+	 * @param SimDataIndex - the sim data index to use when doing the writeback in AppendSimulationData
 	 */
 	virtual void CreateActor(USkeletalMeshComponent* InOwnerComponent, const UClothingAssetBase* InAsset, int32 SimDataIndex) override
 	{
