Chaos Cloth: Fix external convex collision transforms.

#Synced-CL 51472478

diff --git a/Engine/Plugins/ChaosCloth/Source/ChaosCloth/Private/ChaosCloth/ChaosClothingSimulationCollider.cpp b/Engine/Plugins/ChaosCloth/Source/ChaosCloth/Private/ChaosCloth/ChaosClothingSimulationCollider.cpp
--- a/Engine/Plugins/ChaosCloth/Source/ChaosCloth/Private/ChaosCloth/ChaosClothingSimulationCollider.cpp
+++ b/Engine/Plugins/ChaosCloth/Source/ChaosCloth/Private/ChaosCloth/ChaosClothingSimulationCollider.cpp
@@ -244,8 +244,7 @@ void FClothingSimulationCollider::FLODData::Add(
 		{
 			const FClothCollisionPrim_Convex& Convex = ClothCollisionData.Convexes[ConvexIndex];
 
-			// Always initialize the collision particle transforms before setting any geometry as otherwise NaNs gets detected during the bounding box updates
-			BaseTransforms[Index] = Softs::FSolverRigidTransform3::Identity;
+			BaseTransforms[Index] = Softs::FSolverRigidTransform3(Convex.LocalPosition, Convex.LocalRotation);
 
 			BoneIndices[Index] = GetMappedBoneIndex(UsedBoneIndices, Convex.BoneIndex);
 			UE_LOG(LogChaosCloth, VeryVerbose, TEXT("Found collision convex on bone index %d."), BoneIndices[Index]);
diff --git a/Engine/Source/Runtime/ClothingSystemRuntimeInterface/Private/ClothCollisionData.cpp b/Engine/Source/Runtime/ClothingSystemRuntimeInterface/Private/ClothCollisionData.cpp
--- a/Engine/Source/Runtime/ClothingSystemRuntimeInterface/Private/ClothCollisionData.cpp
+++ b/Engine/Source/Runtime/ClothingSystemRuntimeInterface/Private/ClothCollisionData.cpp
@@ -58,10 +58,8 @@ void FClothCollisionData::AppendTransformed(const FClothCollisionData& InOther,
 	for (int32 Index = NumConvexesBefore; Index < Convexes.Num(); ++Index)
 	{
 		const FTransform& Transform = BoneTransforms[Convexes[Index].BoneIndex];
-		for (FVector& SurfacePoint : Convexes[Index].SurfacePoints)
-		{
-			SurfacePoint = Transform.TransformPosition(SurfacePoint);
-		}
+		Convexes[Index].LocalPosition = Transform.TransformPosition(Convexes[Index].LocalPosition);
+		Convexes[Index].LocalRotation = Transform.TransformRotation(Convexes[Index].LocalRotation);
 	}
 	for (int32 Index = NumBoxesBefore; Index < Boxes.Num(); ++Index)
 	{
diff --git a/Engine/Source/Runtime/ClothingSystemRuntimeInterface/Public/ClothCollisionPrim.h b/Engine/Source/Runtime/ClothingSystemRuntimeInterface/Public/ClothCollisionPrim.h
--- a/Engine/Source/Runtime/ClothingSystemRuntimeInterface/Public/ClothCollisionPrim.h
+++ b/Engine/Source/Runtime/ClothingSystemRuntimeInterface/Public/ClothCollisionPrim.h
@@ -83,7 +83,7 @@ struct FClothCollisionPrim_ConvexFace
 	}
 
 	UPROPERTY()
-	FPlane Plane;
+	FPlane Plane; 
 
 	UPROPERTY()
 	TArray<int32> Indices;
@@ -106,15 +106,21 @@ struct FClothCollisionPrim_Convex
 	FClothCollisionPrim_Convex(
 		TArray<FClothCollisionPrim_ConvexFace>&& InFaces,
 		TArray<FVector>&& InSurfacePoints,
-		int32 InBoneIndex = INDEX_NONE)
+		int32 InBoneIndex = INDEX_NONE,
+		const FVector& InLocalPosition = FVector::ZeroVector,
+		const FQuat& InLocalRotation = FQuat::Identity)
 		: Faces(InFaces)
 		, SurfacePoints(InSurfacePoints)
+		, LocalPosition(InLocalPosition)
+		, LocalRotation(InLocalRotation)
 		, BoneIndex(InBoneIndex)
 	{}
 
 	bool operator==(const FClothCollisionPrim_Convex& Other) const
 	{
 		if (BoneIndex == Other.BoneIndex &&
+			LocalPosition == Other.LocalPosition &&
+			LocalRotation == Other.LocalRotation && 
 			Faces == Other.Faces &&
 			SurfacePoints.Num() == Other.SurfacePoints.Num())
 		{
@@ -141,6 +147,12 @@ struct FClothCollisionPrim_Convex
 	UPROPERTY()
 	TArray<FVector> SurfacePoints;  // Surface points, used by Chaos and also for visualization
 
+	UPROPERTY()
+	FVector LocalPosition = FVector::ZeroVector;
+
+	UPROPERTY()
+	FQuat LocalRotation = FQuat::Identity;
+
 	UPROPERTY()
 	int32 BoneIndex;
 };
